Remove hard-coded reference to game plugins directory in template.

diff --git a/Engine/Plugins/Editor/PluginCreator/Templates/Advanced/PluginModule.Build.cs.template b/Engine/Plugins/Editor/PluginCreator/Templates/Advanced/PluginModule.Build.cs.template
--- a/Engine/Plugins/Editor/PluginCreator/Templates/Advanced/PluginModule.Build.cs.template
+++ b/Engine/Plugins/Editor/PluginCreator/Templates/Advanced/PluginModule.Build.cs.template
@@ -38,6 +38,7 @@ public class %MODULE_NAME% : ModuleRules
 		PrivateDependencyModuleNames.AddRange(
 			new string[]
 			{
+				"Projects",
 				"InputCore",
 				"UnrealEd",
 				"LevelEditor",
diff --git a/Engine/Plugins/Editor/PluginCreator/Templates/Advanced/PluginStyle.cpp.template b/Engine/Plugins/Editor/PluginCreator/Templates/Advanced/PluginStyle.cpp.template
--- a/Engine/Plugins/Editor/PluginCreator/Templates/Advanced/PluginStyle.cpp.template
+++ b/Engine/Plugins/Editor/PluginCreator/Templates/Advanced/PluginStyle.cpp.template
@@ -3,6 +3,7 @@
 
 #include "%PLUGIN_NAME%Style.h"
 #include "SlateGameResources.h"
+#include "IPluginManager.h"
 
 TSharedPtr< FSlateStyleSet > F%PLUGIN_NAME%Style::StyleInstance = NULL;
 
@@ -41,7 +42,7 @@ const FVector2D Icon40x40(40.0f, 40.0f);
 TSharedRef< FSlateStyleSet > F%PLUGIN_NAME%Style::Create()
 {
 	TSharedRef< FSlateStyleSet > Style = MakeShareable(new FSlateStyleSet("%PLUGIN_NAME%Style"));
-	Style->SetContentRoot(FPaths::GamePluginsDir() / TEXT("%PLUGIN_NAME%/Resources"));
+	Style->SetContentRoot(IPluginManager::Get().FindPlugin("%PLUGIN_NAME%")->GetBaseDir() / TEXT("Resources"));
 
 	Style->Set("%PLUGIN_NAME%.OpenPluginWindow", new IMAGE_BRUSH(TEXT("ButtonIcon_40x"), Icon40x40));
 
diff --git a/Engine/Plugins/Editor/PluginCreator/Templates/Basic/PluginModule.Build.cs.template b/Engine/Plugins/Editor/PluginCreator/Templates/Basic/PluginModule.Build.cs.template
--- a/Engine/Plugins/Editor/PluginCreator/Templates/Basic/PluginModule.Build.cs.template
+++ b/Engine/Plugins/Editor/PluginCreator/Templates/Basic/PluginModule.Build.cs.template
@@ -38,6 +38,7 @@ public class %MODULE_NAME% : ModuleRules
 		PrivateDependencyModuleNames.AddRange(
 			new string[]
 			{
+				"Projects",
 				"InputCore",
 				"UnrealEd",
 				"LevelEditor",
diff --git a/Engine/Plugins/Editor/PluginCreator/Templates/Basic/PluginStyle.cpp.template b/Engine/Plugins/Editor/PluginCreator/Templates/Basic/PluginStyle.cpp.template
--- a/Engine/Plugins/Editor/PluginCreator/Templates/Basic/PluginStyle.cpp.template
+++ b/Engine/Plugins/Editor/PluginCreator/Templates/Basic/PluginStyle.cpp.template
@@ -3,6 +3,7 @@
 
 #include "%PLUGIN_NAME%Style.h"
 #include "SlateGameResources.h"
+#include "IPluginManager.h"
 
 TSharedPtr< FSlateStyleSet > F%PLUGIN_NAME%Style::StyleInstance = NULL;
 
@@ -41,7 +42,7 @@ const FVector2D Icon40x40(40.0f, 40.0f);
 TSharedRef< FSlateStyleSet > F%PLUGIN_NAME%Style::Create()
 {
 	TSharedRef< FSlateStyleSet > Style = MakeShareable(new FSlateStyleSet("%PLUGIN_NAME%Style"));
-	Style->SetContentRoot(FPaths::GamePluginsDir() / TEXT("%PLUGIN_NAME%/Resources"));
+	Style->SetContentRoot(IPluginManager::Get().FindPlugin("%PLUGIN_NAME%")->GetBaseDir() / TEXT("Resources"));
 
 	Style->Set("%PLUGIN_NAME%.PluginAction", new IMAGE_BRUSH(TEXT("ButtonIcon_40x"), Icon40x40));
 
