Improved instructions for what to do after creating plugins

Prompt after creation to restart the editor or reopen code solution to start editing source files.
Added messages to basic and advanced templates to explain where to override behavior.

diff --git a/Engine/Plugins/Editor/PluginBrowser/Source/PluginBrowser/Private/SNewPluginWizard.cpp b/Engine/Plugins/Editor/PluginBrowser/Source/PluginBrowser/Private/SNewPluginWizard.cpp
--- a/Engine/Plugins/Editor/PluginBrowser/Source/PluginBrowser/Private/SNewPluginWizard.cpp
+++ b/Engine/Plugins/Editor/PluginBrowser/Source/PluginBrowser/Private/SNewPluginWizard.cpp
@@ -568,6 +568,10 @@ FReply SNewPluginWizard::OnCreatePluginClicked()
 			PluginBrowserModule.SetPluginPendingEnableState(AutoPluginName, false, true);
 		}
 
+		FText DialogTitle = LOCTEXT("PluginCreatedTitle", "New Plugin Created");
+		FText DialogText = LOCTEXT("PluginCreatedText", "Restart the editor to activate new Plugin or reopen your code solution to see/edit the newly created source files.");
+		FMessageDialog::Open(EAppMsgType::Ok, DialogText, &DialogTitle);
+
 		auto PinnedOwnerTab = OwnerTab.Pin();
 		if (PinnedOwnerTab.IsValid())
 		{
diff --git a/Engine/Plugins/Editor/PluginBrowser/Templates/Advanced/PluginModule.cpp.template b/Engine/Plugins/Editor/PluginBrowser/Templates/Advanced/PluginModule.cpp.template
--- a/Engine/Plugins/Editor/PluginBrowser/Templates/Advanced/PluginModule.cpp.template
+++ b/Engine/Plugins/Editor/PluginBrowser/Templates/Advanced/PluginModule.cpp.template
@@ -71,11 +71,23 @@ void F%PLUGIN_NAME%Module::ShutdownModule()
 
 TSharedRef<SDockTab> F%PLUGIN_NAME%Module::OnSpawnPluginTab(const FSpawnTabArgs& SpawnTabArgs)
 {
+	FText WidgetText = FText::Format(
+		LOCTEXT("WindowWidgetText", "Add code to {0} in {1} to override this window's contents"),
+		FText::FromString(TEXT("F%PLUGIN_NAME%Module::OnSpawnPluginTab")),
+		FText::FromString(TEXT("%PLUGIN_NAME%.cpp"))
+		);
+
 	return SNew(SDockTab)
 		.TabRole(ETabRole::NomadTab)
 		[
 			// Put your tab content here!
-			SNullWidget::NullWidget
+			SNew(SBox)
+			.HAlign(HAlign_Center)
+			.VAlign(VAlign_Center)
+			[
+				SNew(STextBlock)
+				.Text(WidgetText)
+			]
 		];
 }
 
diff --git a/Engine/Plugins/Editor/PluginBrowser/Templates/Basic/PluginModule.cpp.template b/Engine/Plugins/Editor/PluginBrowser/Templates/Basic/PluginModule.cpp.template
--- a/Engine/Plugins/Editor/PluginBrowser/Templates/Basic/PluginModule.cpp.template
+++ b/Engine/Plugins/Editor/PluginBrowser/Templates/Basic/PluginModule.cpp.template
@@ -66,6 +66,12 @@ void F%PLUGIN_NAME%Module::ShutdownModule()
 void F%PLUGIN_NAME%Module::PluginButtonClicked()
 {
 	// Put your "OnButtonClicked" stuff here
+	FText DialogText = FText::Format(
+							LOCTEXT("PluginButtonDialogText", "Add code to {0} in {1} to override this button's actions"),
+							FText::FromString(TEXT("F%PLUGIN_NAME%Module::PluginButtonClicked()")),
+							FText::FromString(TEXT("%PLUGIN_NAME%.cpp"))
+					   );
+	FMessageDialog::Open(EAppMsgType::Ok, DialogText);
 }
 
 void F%PLUGIN_NAME%Module::AddMenuExtension(FMenuBuilder& Builder)
